Create new MLAdapterInference Plugin to combine MLAdapter and NNI

#rnx

diff --git a/Engine/Plugins/AI/MLAdapter/Source/MLAdapter/Private/MLAdapterManager.cpp b/Engine/Plugins/AI/MLAdapter/Source/MLAdapter/Private/MLAdapterManager.cpp
--- a/Engine/Plugins/AI/MLAdapter/Source/MLAdapter/Private/MLAdapterManager.cpp
+++ b/Engine/Plugins/AI/MLAdapter/Source/MLAdapter/Private/MLAdapterManager.cpp
@@ -345,41 +345,49 @@ void UMLAdapterManager::BindToDelegates()
 #endif // WITH_EDITORONLY_DATA
 }
 
-void UMLAdapterManager::OnPostWorldInit(UWorld* World, const UWorld::InitializationValues)
+bool UMLAdapterManager::ShouldInitForWorld(const UWorld& World) const
 {
-	if (World && World->IsGameWorld())
+	if (!World.IsGameWorld())
 	{
+		return false;
+	}
+
 #if WITH_EDITOR
-		// PIE is a special case, we need to see if it's a client-server PIE and 
-		// if so we need to filter the incoming world based on the settings
-		if (World->WorldType == EWorldType::PIE)
+	// PIE is a special case, we need to see if it's a client-server PIE and 
+	// if so we need to filter the incoming world based on the settings
+	if (World.WorldType == EWorldType::PIE)
+	{
+		if (World.HasAnyFlags(RF_WasLoaded))
 		{
-			if (World->HasAnyFlags(RF_WasLoaded))
+			const ULevelEditorPlaySettings* PlayInSettings = GetDefault<ULevelEditorPlaySettings>();
+			if (PlayInSettings)
 			{
-				const ULevelEditorPlaySettings* PlayInSettings = GetDefault<ULevelEditorPlaySettings>();
-				if (PlayInSettings)
-				{
-					EPlayNetMode PlayNetMode = PIE_Standalone;
-					PlayInSettings->GetPlayNetMode(PlayNetMode);
-
-					if (PlayNetMode != PIE_Standalone)
-					{
-						UE_LOG(LogUnrealEditorMLAdapter, Log, TEXT("Ignoring %s due to net mode != PIE_Standalone")
-							, *World->GetName());
-						//UMLAdapterSettings::
-						return;
-					}
-				}
-				else
+				EPlayNetMode PlayNetMode = PIE_Standalone;
+				PlayInSettings->GetPlayNetMode(PlayNetMode);
+
+				if (PlayNetMode != PIE_Standalone)
 				{
-					// skipping temp world
-					return;
+					UE_LOG(LogUnrealEditorMLAdapter, Log, TEXT("Ignoring %s due to net mode != PIE_Standalone"), *World.GetName());
+					return false;
 				}
 			}
+			else
+			{
+				// skipping temp world
+				return false;
+			}
+		}
 
-		}		
+	}
 #endif // WITH_EDITOR
 
+	return true;
+}
+
+void UMLAdapterManager::OnPostWorldInit(UWorld* World, const UWorld::InitializationValues)
+{
+	if (World && ShouldInitForWorld(*World))
+	{
 		LastActiveWorld = World;
 		if ((RequestedFunctionMode == EMLAdapterServerMode::Invalid || RequestedFunctionMode == EMLAdapterServerMode::AutoDetect)
 			&& GetServerModeForWorld(*World) != CurrentFunctionMode)
diff --git a/Engine/Plugins/AI/MLAdapter/Source/MLAdapter/Public/Agents/MLAdapterAgent.h b/Engine/Plugins/AI/MLAdapter/Source/MLAdapter/Public/Agents/MLAdapterAgent.h
--- a/Engine/Plugins/AI/MLAdapter/Source/MLAdapter/Public/Agents/MLAdapterAgent.h
+++ b/Engine/Plugins/AI/MLAdapter/Source/MLAdapter/Public/Agents/MLAdapterAgent.h
@@ -23,7 +23,7 @@ struct FMLAdapterParameterMap
 {
 	GENERATED_BODY()
 
-	UPROPERTY()
+	UPROPERTY(EditAnywhere, Category = MLAdapter)
 	TMap<FName, FString> Params;
 };
 
diff --git a/Engine/Plugins/AI/MLAdapter/Source/MLAdapter/Public/MLAdapterManager.h b/Engine/Plugins/AI/MLAdapter/Source/MLAdapter/Public/MLAdapterManager.h
--- a/Engine/Plugins/AI/MLAdapter/Source/MLAdapter/Public/MLAdapterManager.h
+++ b/Engine/Plugins/AI/MLAdapter/Source/MLAdapter/Public/MLAdapterManager.h
@@ -119,6 +119,8 @@ protected:
 
 	void AddCommonFunctions(FRPCServer& Server);
 
+	virtual bool ShouldInitForWorld(const UWorld& World) const;
+
 protected:	
 	friend struct FMLAdapterConsoleCommands::FHelper;
 
diff --git a/Engine/Plugins/Experimental/MLAdapterInference/MLAdapterInference.uplugin b/Engine/Plugins/Experimental/MLAdapterInference/MLAdapterInference.uplugin
new file mode 100644
--- /dev/null
+++ b/Engine/Plugins/Experimental/MLAdapterInference/MLAdapterInference.uplugin
@@ -0,0 +1,35 @@
+{
+	"FileVersion": 3,
+	"Version": 1,
+	"VersionName": "0.0.1",
+	"FriendlyName": "In-Engine Inference Extensions for UE ML Adapter",
+	"Description": "[EXPERIMENTAL] Extends ML Adapter functionality by integrating Unreal's NeuralNetworkInference Plugin",
+	"Category": "Experimental",
+	"CreatedBy": "Epic Games, Inc.",
+	"CreatedByURL": "https://epicgames.com",
+	"DocsURL": "",
+	"MarketplaceURL": "",
+	"SupportURL": "",
+	"EnabledByDefault": false,
+	"CanContainContent": false,
+	"IsBetaVersion": true,
+	"IsExperimentalVersion": true,
+	"Installed": false,
+	"Modules": [
+		{
+			"Name": "MLAdapterInference",
+			"Type": "Runtime",
+			"LoadingPhase": "PreDefault"
+		}
+	],
+	"Plugins": [
+		{
+			"Name": "MLAdapter",
+			"Enabled": true
+		},
+		{
+			"Name": "NeuralNetworkInference",
+			"Enabled": true
+		}
+	]
+}
\ No newline at end of file
diff --git a/Engine/Plugins/Experimental/MLAdapterInference/README.md b/Engine/Plugins/Experimental/MLAdapterInference/README.md
new file mode 100644
--- /dev/null
+++ b/Engine/Plugins/Experimental/MLAdapterInference/README.md
@@ -0,0 +1,7 @@
+# MLAdapterInference
+###### Please note that this plugin is very much experimental and work-in-progress.
+
+This plugin extends the functionality of MLAdapter by adding inference specific classes. Once you have finished developing an ML model using MLAdapter and are ready to ship it in a game, you most likely no longer want to have a dependency on the external Python process. This plugin adds an UMLAdapterInferenceManager class so that you will no longer start an RPC server, and also allows you to directly evaluate your model in engine by using Epic's Neural Network Inference (NNI) Engine. This functionality is provided through UMLAdapterInferenceAgent.
+
+## Feedback
+Please send your feedback to [REDACTED_EMAIL]
\ No newline at end of file
diff --git a/Engine/Plugins/Experimental/MLAdapterInference/Source/MLAdapterInference/MLAdapterInference.Build.cs b/Engine/Plugins/Experimental/MLAdapterInference/Source/MLAdapterInference/MLAdapterInference.Build.cs
new file mode 100644
--- /dev/null
+++ b/Engine/Plugins/Experimental/MLAdapterInference/Source/MLAdapterInference/MLAdapterInference.Build.cs
@@ -0,0 +1,59 @@
+// Copyright Epic Games, Inc. All Rights Reserved.
+
+using System.IO;
+
+namespace UnrealBuildTool.Rules
+{
+	public class MLAdapterInference : ModuleRules
+	{
+		public MLAdapterInference(ReadOnlyTargetRules Target) : base(Target)
+		{
+			PCHUsage = ModuleRules.PCHUsageMode.UseExplicitOrSharedPCHs;
+
+			PrivateIncludePaths.Add("MLAdapterInference/Private");
+
+			PublicIncludePaths.AddRange(
+				new string[] {
+							"Runtime/AIModule/Public",
+							ModuleDirectory + "/Public",
+				}
+			);
+
+			PublicDependencyModuleNames.AddRange(
+				new string[] {
+					"Core",
+					"MLAdapter",
+					"NeuralNetworkInference",
+				}
+			);
+
+			PrivateDependencyModuleNames.AddRange(
+                new string[] {
+					"CoreUObject",
+					"DeveloperSettings",
+					"Engine",
+				}
+			);
+
+			// RPCLib disabled on other platforms at the moment
+			if (Target.Platform == UnrealTargetPlatform.Win64)
+			{
+				PublicDefinitions.Add("WITH_RPCLIB=1");
+				AddEngineThirdPartyPrivateStaticDependencies(Target, "RPCLib");
+
+				string RPClibDir = Path.Combine(Target.UEThirdPartySourceDirectory, "rpclib");
+				PublicIncludePaths.Add(Path.Combine(RPClibDir, "Source", "include"));
+			}
+			else
+			{
+				PublicDefinitions.Add("WITH_RPCLIB=0");
+			}
+
+			if (Target.bBuildEditor == true)
+			{
+				PrivateDependencyModuleNames.Add("EditorFramework");
+				PrivateDependencyModuleNames.Add("UnrealEd");
+			}
+		}
+	}
+}
diff --git a/Engine/Plugins/Experimental/MLAdapterInference/Source/MLAdapterInference/Private/MLAdapterInferenceAgent.cpp b/Engine/Plugins/Experimental/MLAdapterInference/Source/MLAdapterInference/Private/MLAdapterInferenceAgent.cpp
new file mode 100644
--- /dev/null
+++ b/Engine/Plugins/Experimental/MLAdapterInference/Source/MLAdapterInference/Private/MLAdapterInferenceAgent.cpp
@@ -0,0 +1,29 @@
+// Copyright Epic Games, Inc. All Rights Reserved.
+
+#include "MLAdapterInferenceAgent.h"
+#include "MLAdapterTypes.h"
+#include "MLAdapterInferenceTypes.h"
+#include "HAL/UnrealMemory.h"
+
+void UMLAdapterInferenceAgent::Think(const float DeltaTime)
+{
+	if (Brain == nullptr)
+	{
+		UE_LOG(LogMLAdapterInference, Warning, TEXT("Agent beginning to Think but Brain is null"));
+		return;
+	}
+
+	TArray<uint8> Buffer;
+	FMLAdapterMemoryWriter Writer(Buffer);
+	GetObservations(Writer);
+
+	const float* DataPtr = (float*)Buffer.GetData();
+	FMemory::Memcpy(Brain->GetInputDataPointerMutable(), DataPtr, Buffer.Num());
+
+	Brain->Run();
+
+	const FNeuralTensor& Tensor = Brain->GetOutputTensor();
+	const TArray<uint8>& Data = Tensor.GetUnderlyingUInt8ArrayRef();
+	FMLAdapterMemoryReader Reader(Data);
+	DigestActions(Reader);
+}
diff --git a/Engine/Plugins/Experimental/MLAdapterInference/Source/MLAdapterInference/Private/MLAdapterInferenceManager.cpp b/Engine/Plugins/Experimental/MLAdapterInference/Source/MLAdapterInference/Private/MLAdapterInferenceManager.cpp
new file mode 100644
--- /dev/null
+++ b/Engine/Plugins/Experimental/MLAdapterInference/Source/MLAdapterInference/Private/MLAdapterInferenceManager.cpp
@@ -0,0 +1,49 @@
+// Copyright Epic Games, Inc. All Rights Reserved.
+
+#include "MLAdapterInferenceManager.h"
+#include "MLAdapterInferenceAgent.h"
+#include "MLAdapterInferenceSettings.h"
+#include "MLAdapterInferenceTypes.h"
+#include "NeuralNetwork.h"
+#include "MLAdapterSession.h"
+
+
+#if WITH_EDITORONLY_DATA
+#include "Editor.h"
+#include "Settings/LevelEditorPlaySettings.h"
+#endif // WITH_EDITORONLY_DATA
+
+void UMLAdapterInferenceManager::OnPostWorldInit(UWorld* World, const UWorld::InitializationValues)
+{
+	if (World && ShouldInitForWorld(*World))
+	{
+		LastActiveWorld = World;
+
+		if (HasSession())
+		{
+			GetSession().OnPostWorldInit(*World);
+		}
+
+		// Create agents from config
+		const TArrayView<const FInferenceAgentConfig>& AgentConfigs = UMLAdapterInferenceSettings::GetAgentConfigs();
+
+		if (AgentConfigs.IsEmpty())
+		{
+			UE_LOG(LogMLAdapterInference, Warning, TEXT("AgentConfigs is empty so no agents will be constructed. Consider setting Edit->ProjectSettings->Engine->MLAdapter->AgentConfigs"));
+		}
+
+		for (const FInferenceAgentConfig& InferenceAgentConfig : AgentConfigs)
+		{
+			FMLAdapterAgentConfig AdapterAgentConfig = InferenceAgentConfig.AsMLAdapterAgentConfig();
+
+			FMLAdapter::FAgentID AgentID = GetSession().AddAgent(AdapterAgentConfig);
+			UE_LOG(LogMLAdapterInference, Log, TEXT("Created new agent of class %s with AgentID %d"), *AdapterAgentConfig.AgentClassName.ToString(), AgentID);
+
+			UNeuralNetwork* Brain = (UNeuralNetwork*)InferenceAgentConfig.NeuralNetworkPath.TryLoad();
+			checkf(Brain != nullptr, TEXT("Unable to load UNeuralNetwork from NeuralNetworkPath"));
+
+			UMLAdapterInferenceAgent* Agent = Cast<UMLAdapterInferenceAgent>(GetSession().GetAgent(AgentID));
+			Agent->Brain = Brain;
+		}
+	}
+}
diff --git a/Engine/Plugins/Experimental/MLAdapterInference/Source/MLAdapterInference/Private/MLAdapterInferenceModule.cpp b/Engine/Plugins/Experimental/MLAdapterInference/Source/MLAdapterInference/Private/MLAdapterInferenceModule.cpp
new file mode 100644
--- /dev/null
+++ b/Engine/Plugins/Experimental/MLAdapterInference/Source/MLAdapterInference/Private/MLAdapterInferenceModule.cpp
@@ -0,0 +1,18 @@
+// Copyright Epic Games, Inc. All Rights Reserved.
+
+#include "MLAdapterInferenceModule.h"
+#include "Modules/ModuleManager.h"
+
+#define LOCTEXT_NAMESPACE "MLAdapterInferenceModule"
+
+IMPLEMENT_MODULE(FMLAdapterInferenceModule, MLAdapterInferenceModule)
+
+void FMLAdapterInferenceModule::StartupModule()
+{
+}
+
+void FMLAdapterInferenceModule::ShutdownModule()
+{
+}
+
+#undef LOCTEXT_NAMESPACE
diff --git a/Engine/Plugins/Experimental/MLAdapterInference/Source/MLAdapterInference/Private/MLAdapterInferenceSettings.cpp b/Engine/Plugins/Experimental/MLAdapterInference/Source/MLAdapterInference/Private/MLAdapterInferenceSettings.cpp
new file mode 100644
--- /dev/null
+++ b/Engine/Plugins/Experimental/MLAdapterInference/Source/MLAdapterInference/Private/MLAdapterInferenceSettings.cpp
@@ -0,0 +1,55 @@
+// Copyright Epic Games, Inc. All Rights Reserved.
+
+#include "MLAdapterInferenceSettings.h"
+#include "GameFramework/PlayerController.h"
+#include "MLAdapterInferenceAgent.h"
+
+#define GET_CONFIG_VALUE(a) (GetDefault<UMLAdapterInferenceSettings>()->a)
+TConstArrayView<FInferenceAgentConfig> UMLAdapterInferenceSettings::GetAgentConfigs()
+{
+	const TArray<FInferenceAgentConfig>& Configs = GET_CONFIG_VALUE(AgentConfigs);
+	TConstArrayView<FInferenceAgentConfig> View((const FInferenceAgentConfig*)Configs.GetData(), Configs.Num());
+	return View;
+}
+#undef GET_CONFIG_VALUE
+
+FMLAdapterAgentConfig FInferenceAgentConfig::AsMLAdapterAgentConfig() const
+{
+	FMLAdapterAgentConfig ResultConfig;
+
+	if (AgentClass != nullptr)
+	{
+		ResultConfig.AgentClassName = AgentClass->GetFName();
+	}
+	else
+	{
+		ResultConfig.AgentClassName = UMLAdapterInferenceAgent::StaticClass()->GetFName();
+	}
+
+	if (AvatarClass != nullptr)
+	{
+		ResultConfig.AvatarClassName = AvatarClass->GetFName();
+		ResultConfig.AvatarClass = AvatarClass;
+	}
+	else
+	{
+		ResultConfig.AvatarClassName = APlayerController::StaticClass()->GetFName();
+		ResultConfig.AvatarClass = APlayerController::StaticClass();
+	}	
+
+	for (FInferenceSensorConfig SensorConfig : Sensors)
+	{
+		FMLAdapterParameterMap ParamMap;
+		ParamMap.Params = SensorConfig.Params;
+		ResultConfig.Sensors.Add(FName(SensorConfig.SensorClass.GetAssetName()), ParamMap);
+	}
+
+	for (FInferenceActuatorConfig ActuatorConfig : Actuators)
+	{
+		FMLAdapterParameterMap ParamMap;
+		ParamMap.Params = ActuatorConfig.Params;
+		ResultConfig.Actuators.Add(FName(ActuatorConfig.ActuatorClass.GetAssetName()), ParamMap);
+	}
+
+	return ResultConfig;
+}
diff --git a/Engine/Plugins/Experimental/MLAdapterInference/Source/MLAdapterInference/Private/MLAdapterInferenceTypes.cpp b/Engine/Plugins/Experimental/MLAdapterInference/Source/MLAdapterInference/Private/MLAdapterInferenceTypes.cpp
new file mode 100644
--- /dev/null
+++ b/Engine/Plugins/Experimental/MLAdapterInference/Source/MLAdapterInference/Private/MLAdapterInferenceTypes.cpp
@@ -0,0 +1,5 @@
+// Copyright Epic Games, Inc. All Rights Reserved.
+
+#include "MLAdapterInferenceTypes.h"
+
+DEFINE_LOG_CATEGORY(LogMLAdapterInference)
diff --git a/Engine/Plugins/Experimental/MLAdapterInference/Source/MLAdapterInference/Public/MLAdapterInferenceAgent.h b/Engine/Plugins/Experimental/MLAdapterInference/Source/MLAdapterInference/Public/MLAdapterInferenceAgent.h
new file mode 100644
--- /dev/null
+++ b/Engine/Plugins/Experimental/MLAdapterInference/Source/MLAdapterInference/Public/MLAdapterInferenceAgent.h
@@ -0,0 +1,18 @@
+// Copyright Epic Games, Inc. All Rights Reserved.
+
+#pragma once
+
+#include "Agents/MLAdapterAgent.h"
+#include "NeuralNetwork.h"
+#include "MLAdapterInferenceAgent.generated.h"
+
+UCLASS(Blueprintable, EditInlineNew)
+class MLADAPTERINFERENCE_API UMLAdapterInferenceAgent : public UMLAdapterAgent
+{
+	GENERATED_BODY()
+public:
+	virtual void Think(const float DeltaTime) override;
+
+	UPROPERTY(EditAnywhere, BlueprintReadWrite, Category = InferenceAgent)
+	UNeuralNetwork* Brain;
+};
diff --git a/Engine/Plugins/Experimental/MLAdapterInference/Source/MLAdapterInference/Public/MLAdapterInferenceManager.h b/Engine/Plugins/Experimental/MLAdapterInference/Source/MLAdapterInference/Public/MLAdapterInferenceManager.h
new file mode 100644
--- /dev/null
+++ b/Engine/Plugins/Experimental/MLAdapterInference/Source/MLAdapterInference/Public/MLAdapterInferenceManager.h
@@ -0,0 +1,18 @@
+// Copyright Epic Games, Inc. All Rights Reserved.
+
+#pragma once 
+
+#include "CoreMinimal.h"
+#include "MLAdapterManager.h"
+#include "MLAdapterInferenceManager.generated.h"
+
+/**
+ * Inference manager is only doing inference so we won't startup the RPC server.
+ */
+UCLASS()
+class MLADAPTERINFERENCE_API UMLAdapterInferenceManager : public UMLAdapterManager
+{
+	GENERATED_BODY()
+public:
+	virtual void OnPostWorldInit(UWorld* World, const UWorld::InitializationValues) override;
+};
diff --git a/Engine/Plugins/Experimental/MLAdapterInference/Source/MLAdapterInference/Public/MLAdapterInferenceModule.h b/Engine/Plugins/Experimental/MLAdapterInference/Source/MLAdapterInference/Public/MLAdapterInferenceModule.h
new file mode 100644
--- /dev/null
+++ b/Engine/Plugins/Experimental/MLAdapterInference/Source/MLAdapterInference/Public/MLAdapterInferenceModule.h
@@ -0,0 +1,14 @@
+// Copyright Epic Games, Inc. All Rights Reserved.
+
+#pragma once
+
+#include "CoreMinimal.h"
+#include "Modules/ModuleManager.h"
+
+class MLADAPTERINFERENCE_API FMLAdapterInferenceModule : public IModuleInterface
+{
+public:
+	/** IModuleInterface implementation. */
+	virtual void StartupModule() override;
+	virtual void ShutdownModule() override;
+};
diff --git a/Engine/Plugins/Experimental/MLAdapterInference/Source/MLAdapterInference/Public/MLAdapterInferenceSettings.h b/Engine/Plugins/Experimental/MLAdapterInference/Source/MLAdapterInference/Public/MLAdapterInferenceSettings.h
new file mode 100644
--- /dev/null
+++ b/Engine/Plugins/Experimental/MLAdapterInference/Source/MLAdapterInference/Public/MLAdapterInferenceSettings.h
@@ -0,0 +1,75 @@
+// Copyright Epic Games, Inc. All Rights Reserved.
+
+#pragma once
+
+#include "CoreMinimal.h"
+#include "Containers/ArrayView.h"
+#include "Engine/DeveloperSettings.h"
+#include "GameFramework/Controller.h"
+#include "Agents/MLAdapterAgent.h"
+#include "Sensors/MLAdapterSensor.h"
+#include "Actuators/MLAdapterActuator.h"
+#include "NeuralNetwork.h"
+#include "MLAdapterInferenceSettings.generated.h"
+
+
+USTRUCT()
+struct MLADAPTERINFERENCE_API FInferenceSensorConfig : public FMLAdapterParameterMap
+{
+	GENERATED_BODY()
+
+	UPROPERTY(EditAnywhere, Category = MLAdapterInference, meta = (MetaClass = "MLAdapterSensor"))
+	FSoftClassPath SensorClass;
+};
+
+USTRUCT()
+struct MLADAPTERINFERENCE_API FInferenceActuatorConfig : public FMLAdapterParameterMap
+{
+	GENERATED_BODY()
+
+	UPROPERTY(EditAnywhere, Category = MLAdapterInference, meta = (MetaClass = "MLAdapterActuator"))
+	FSoftClassPath ActuatorClass;
+};
+
+USTRUCT()
+struct MLADAPTERINFERENCE_API FInferenceAgentConfig
+{
+	GENERATED_BODY()
+
+	UPROPERTY(EditAnywhere, config, Category = MLAdapterInference)
+	TSubclassOf<UMLAdapterAgent> AgentClass;
+
+	UPROPERTY(EditAnywhere, config, Category = MLAdapterInference)
+	TSubclassOf<AActor> AvatarClass;
+
+	UPROPERTY(EditAnywhere, config, Category = MLAdapterInference, meta = (AllowedClasses = "NeuralNetwork"))
+	FSoftObjectPath NeuralNetworkPath;
+
+	UPROPERTY(EditAnywhere, config, Category = MLAdapterInference)
+	TArray<FInferenceSensorConfig> Sensors;
+
+	UPROPERTY(EditAnywhere, config, Category = MLAdapterInference)
+	TArray<FInferenceActuatorConfig> Actuators;
+
+	FMLAdapterAgentConfig AsMLAdapterAgentConfig() const;
+};
+
+#define GET_CONFIG_VALUE(a) (GetDefault<UMLAdapterInferenceSettings>()->a)
+
+/**
+ * Implements the settings for the MLAdapterInference plugin.
+ */
+UCLASS(config = Plugins, defaultconfig, DisplayName = "MLAdapterInference")
+class MLADAPTERINFERENCE_API UMLAdapterInferenceSettings : public UDeveloperSettings
+{
+	GENERATED_BODY()
+
+public:
+	static TConstArrayView<FInferenceAgentConfig> GetAgentConfigs();
+
+protected:
+	UPROPERTY(EditAnywhere, config, Category = MLAdapterInference)
+	TArray<FInferenceAgentConfig> AgentConfigs;
+};
+
+#undef GET_CONFIG_VALUE
\ No newline at end of file
diff --git a/Engine/Plugins/Experimental/MLAdapterInference/Source/MLAdapterInference/Public/MLAdapterInferenceTypes.h b/Engine/Plugins/Experimental/MLAdapterInference/Source/MLAdapterInference/Public/MLAdapterInferenceTypes.h
new file mode 100644
--- /dev/null
+++ b/Engine/Plugins/Experimental/MLAdapterInference/Source/MLAdapterInference/Public/MLAdapterInferenceTypes.h
@@ -0,0 +1,7 @@
+// Copyright Epic Games, Inc. All Rights Reserved.
+
+#pragma once
+
+#include "CoreMinimal.h"
+
+DECLARE_LOG_CATEGORY_EXTERN(LogMLAdapterInference, Log, All);
